Chaos Dataflow Solver Actor

#Synced-CL 51222019

#Synced-CL 51222019

diff --git a/Engine/Plugins/Experimental/ChaosDataflowSolver/ChaosDataflowSolver.uplugin b/Engine/Plugins/Experimental/ChaosDataflowSolver/ChaosDataflowSolver.uplugin
new file mode 100644
--- /dev/null
+++ b/Engine/Plugins/Experimental/ChaosDataflowSolver/ChaosDataflowSolver.uplugin
@@ -0,0 +1,24 @@
+{
+	"FileVersion": 3,
+	"Version": 1,
+	"VersionName": "1.0",
+	"FriendlyName": "ChaosDataflowSolver",
+	"Description": "",
+	"Category": "Other",
+	"CreatedBy": "",
+	"CreatedByURL": "",
+	"DocsURL": "",
+	"MarketplaceURL": "",
+	"SupportURL": "",
+	"CanContainContent": true,
+	"IsBetaVersion": false,
+	"IsExperimentalVersion": true,
+	"Installed": false,
+	"Modules": [
+		{
+			"Name": "ChaosDataflowSolver",
+			"Type": "Runtime",
+			"LoadingPhase": "Default"
+		}
+	]
+}
\ No newline at end of file
diff --git a/Engine/Plugins/Experimental/ChaosDataflowSolver/Source/ChaosDataflowSolver/ChaosDataflowSolver.Build.cs b/Engine/Plugins/Experimental/ChaosDataflowSolver/Source/ChaosDataflowSolver/ChaosDataflowSolver.Build.cs
new file mode 100644
--- /dev/null
+++ b/Engine/Plugins/Experimental/ChaosDataflowSolver/Source/ChaosDataflowSolver/ChaosDataflowSolver.Build.cs
@@ -0,0 +1,59 @@
+// Copyright Epic Games, Inc. All Rights Reserved.
+
+using UnrealBuildTool;
+
+public class ChaosDataflowSolver : ModuleRules
+{
+	public ChaosDataflowSolver(ReadOnlyTargetRules Target) : base(Target)
+	{
+		PCHUsage = ModuleRules.PCHUsageMode.UseExplicitOrSharedPCHs;
+		
+		PublicIncludePaths.AddRange(
+			new string[] {
+			}
+		);
+				
+		
+		PrivateIncludePaths.AddRange(
+			new string[] {
+			}
+		);
+			
+		
+		PublicDependencyModuleNames.AddRange(
+			new string[]
+			{
+				"Core",
+				"DataflowCore",
+				"DataflowEngine",
+				"DataflowSimulation",
+				"RigidPhysics",
+			}
+		);
+			
+		
+		PrivateDependencyModuleNames.AddRange(
+			new string[]
+			{
+				"CoreUObject",
+				"Engine",
+				"Slate",
+				"SlateCore",
+				"Landscape",
+			}
+		);
+
+		PublicDefinitions.Add("UE_RIGIDPHYSICS_API_ENABLED=1");
+
+		SetupModulePhysicsSupport(Target);
+		PrivateDefinitions.Add("CHAOS_INCLUDE_LEVEL_1=1");
+
+
+		DynamicallyLoadedModuleNames.AddRange(
+			new string[]
+			{
+				// ... add any modules that your module loads dynamically here ...
+			}
+		);
+	}
+}
diff --git a/Engine/Plugins/Experimental/ChaosDataflowSolver/Source/ChaosDataflowSolver/Private/ChaosDataflowSolver.cpp b/Engine/Plugins/Experimental/ChaosDataflowSolver/Source/ChaosDataflowSolver/Private/ChaosDataflowSolver.cpp
new file mode 100644
--- /dev/null
+++ b/Engine/Plugins/Experimental/ChaosDataflowSolver/Source/ChaosDataflowSolver/Private/ChaosDataflowSolver.cpp
@@ -0,0 +1,17 @@
+// Copyright Epic Games, Inc. All Rights Reserved.
+
+#include "ChaosDataflowSolver.h"
+
+#define LOCTEXT_NAMESPACE "FChaosDataflowSolverModule"
+
+void FChaosDataflowSolverModule::StartupModule()
+{
+}
+
+void FChaosDataflowSolverModule::ShutdownModule()
+{
+}
+
+#undef LOCTEXT_NAMESPACE
+	
+IMPLEMENT_MODULE(FChaosDataflowSolverModule, ChaosDataflowSolver)
\ No newline at end of file
diff --git a/Engine/Plugins/Experimental/ChaosDataflowSolver/Source/ChaosDataflowSolver/Private/ChaosDataflowSolverActor.cpp b/Engine/Plugins/Experimental/ChaosDataflowSolver/Source/ChaosDataflowSolver/Private/ChaosDataflowSolverActor.cpp
new file mode 100644
--- /dev/null
+++ b/Engine/Plugins/Experimental/ChaosDataflowSolver/Source/ChaosDataflowSolver/Private/ChaosDataflowSolverActor.cpp
@@ -0,0 +1,279 @@
+// Copyright Epic Games, Inc. All Rights Reserved.
+
+#include "ChaosDataflowSolverActor.h"
+#include "UObject/ConstructorHelpers.h"
+#include "PhysicsSolver.h"
+#include "ChaosModule.h"
+
+#include "Components/BillboardComponent.h"
+#include "EngineUtils.h"
+#include "ChaosSolversModule.h"
+#include "Chaos/ChaosGameplayEventDispatcher.h"
+#include "Chaos/Framework/DebugSubstep.h"
+#include "Engine/Texture2D.h"
+#include "UObject/FortniteMainBranchObjectVersion.h"
+#include "PhysicsProxy/SingleParticlePhysicsProxy.h"
+#include "PhysicsEngine/PhysicsSettings.h"
+#include "Dataflow/DataflowSimulationNodes.h"
+
+#include "RigidPhysics/RigidBody.h"
+#include "RigidPhysics/RigidContext.h"
+#include "RigidPhysics/RigidModifier.h"
+#include "RigidPhysics/RigidScene.h"
+#include "RigidAssetUserData.h"
+
+#include "RigidComponents/PrimitiveRigidComponentInterface.h"
+#include "RigidComponents/InstancedRigidComponentInterface.h"
+#include "RigidComponents/LandscapeRigidComponentInterface.h"
+#include "RigidComponents/SkeletalRigidComponentInterface.h"
+
+#include UE_INLINE_GENERATED_CPP_BY_NAME(ChaosDataflowSolverActor)
+
+AChaosDataflowSolverActor::AChaosDataflowSolverActor(const FObjectInitializer& ObjectInitializer)
+	: Super(ObjectInitializer)
+{
+	PrimaryActorTick.bCanEverTick = true;
+	PrimaryActorTick.bTickEvenWhenPaused = false;
+	PrimaryActorTick.bStartWithTickEnabled = true;
+	PrimaryActorTick.TickGroup = TG_PrePhysics;
+}
+
+void AChaosDataflowSolverActor::BeginPlay()
+{
+	Super::BeginPlay();
+}
+
+void AChaosDataflowSolverActor::EndPlay(const EEndPlayReason::Type ReasonEnd)
+{
+	Super::EndPlay(ReasonEnd);
+
+	if (SolverContext)
+	{
+		if (FRigidPhysicsSceneDataflowState* State = SolverContext->GetRigidState())
+		{
+			if (UE::Physics::FRigidContextGameRW Context = State->Handle.LockRW())
+			{
+				for (TWeakObjectPtr<UPrimitiveComponent> WeakComponent : RegisteredComponents)
+				{
+					if (UPrimitiveComponent* Component = WeakComponent.Get())
+					{
+						URigidAssetUserData* RigidData = Component->GetAssetUserData<URigidAssetUserData>();
+						if (RigidData)
+						{
+							for (UE::Physics::FRigidBodyHandle& Handle : RigidData->Bodies)
+							{
+								Context->DestroyBody(Handle.Pin(Context));
+							}
+						}
+
+						Component->RemoveUserDataOfClass(URigidAssetUserData::StaticClass());
+					}
+				}
+
+				RegisteredComponents.Empty();
+			}
+
+			State->Factory->DestroyScene(State->Handle);
+		}
+	}
+}
+
+void AChaosDataflowSolverActor::BuildSimulationProxy()
+{
+
+}
+void AChaosDataflowSolverActor::ResetSimulationProxy()
+{
+
+}
+void AChaosDataflowSolverActor::WriteToSimulation(const float DeltaTime, const bool bAsyncTask)
+{
+}
+void AChaosDataflowSolverActor::ReadFromSimulation(const float DeltaTime, const bool bAsyncTask)
+{
+	
+}
+
+void AChaosDataflowSolverActor::Tick(float DeltaTime)
+{
+	Super::Tick(DeltaTime);
+
+	if (SolverContext && SimulationAsset.DataflowAsset)
+	{
+		if (FRigidPhysicsSceneDataflowState* State = SolverContext->GetRigidState())
+		{
+			if (UE::Physics::FRigidContextGameRW Context = State->Handle.LockRW())
+			{
+				for (TWeakObjectPtr<UPrimitiveComponent> WeakComponent: RegisteredComponents)
+				{
+					if (UPrimitiveComponent* Component = WeakComponent.Get())
+					{
+						// TODO: Edit logic later when Landscapes and Skeletal are fully implemented
+						bool bBasicType = !Component->IsA<UInstancedStaticMeshComponent>() && !Component->IsA<ULandscapeHeightfieldCollisionComponent>() && !Component->IsA<USkeletalMeshComponent>();
+
+						URigidAssetUserData* RigidData = Component->GetAssetUserData<URigidAssetUserData>();
+						if (bBasicType && RigidData && RigidData->Bodies.Num() > 0)
+						{
+							if (UE::Physics::TRigidBodyPtr<UE::Physics::FRigidContextGameRW> BodyPtr = RigidData->Bodies[0].Pin(Context))
+							{
+								FTransform Transform = Component->GetComponentToWorld();
+
+								if (!BodyPtr->GetTransform().EqualsNoScale(Transform))
+								{
+									if (BodyPtr->IsKinematic())
+									{
+										BodyPtr->SetKinematicTarget(Transform);
+									}
+									else
+									{
+										BodyPtr->UpdateTransform(Transform);
+									}
+								}
+							}
+						}
+					}
+				}
+			}
+
+			
+			if (const TSharedPtr<UE::Dataflow::FGraph> DataflowGraph = SimulationAsset.DataflowAsset->GetDataflowGraph())
+			{
+				const TArray<UE::Dataflow::FGraph::FNodeSharedPtr>& Array = DataflowGraph->GetFilteredNodes(FSceneSimulationTerminalDataflowNode::StaticType());
+
+				SolverContext->SetTimingInfos(DeltaTime, 0);
+				SolverContext->ClearAllData();
+
+				for (const TSharedPtr<FDataflowNode>& ExecutionNode : Array)
+				{
+					SolverContext->Evaluate(ExecutionNode.Get(), nullptr);
+				}
+			}
+
+			if (UE::Physics::FRigidContextGameRO Context = State->Handle.LockRO())
+			{
+				for (TWeakObjectPtr<UPrimitiveComponent> WeakComponent : RegisteredComponents)
+				{
+					if (UPrimitiveComponent* Component = WeakComponent.Get())
+					{
+						if (URigidAssetUserData* RigidData = Component->GetAssetUserData<URigidAssetUserData>())
+						{
+							if (UInstancedStaticMeshComponent* InstComponent = Cast<UInstancedStaticMeshComponent>(Component))
+							{
+							}
+							else if (ULandscapeHeightfieldCollisionComponent* LandscapeComponent = Cast<ULandscapeHeightfieldCollisionComponent>(Component))
+							{
+							}
+							else if (USkeletalMeshComponent* SkelComponent = Cast<USkeletalMeshComponent>(Component))
+							{
+								SkeletalRigidComponentInterface::OnSolverEndFrame(SkelComponent, State->Handle, RigidData->Bodies);
+							}
+							else
+							{
+								PrimitiveRigidComponentInterface::OnSolverEndFrame(Component, State->Handle, RigidData->Bodies);
+							}
+						}
+					}
+				}
+			}
+		}
+	}
+}
+
+void AChaosDataflowSolverActor::SetSolverActive(bool bActive)
+{
+
+}
+
+void AChaosDataflowSolverActor::PostInitializeComponents()
+{
+	Super::PostInitializeComponents();
+	if (SimulationAsset.DataflowAsset) {
+		
+		const TSharedPtr<UE::Dataflow::FGraph> DataflowGraph = SimulationAsset.DataflowAsset->GetDataflowGraph();
+
+		if (DataflowGraph)
+		{
+			SolverContext = MakeUnique<UE::Dataflow::FDataflowSimulationContext>(SimulationAsset.DataflowAsset);
+
+			SolverContext->SetRigidSceneName(GetName());
+
+			const TArray<UE::Dataflow::FGraph::FNodeSharedPtr>& Array = DataflowGraph->GetFilteredNodes(FRigidSceneSetupTerminalDataflowNode::StaticType());
+
+			if(!Array.IsEmpty())
+			{
+				for (const TSharedPtr<FDataflowNode>& ExecutionNode : Array)
+				{
+					SolverContext->Evaluate(ExecutionNode.Get(), nullptr);
+				}
+			}
+		}
+	}
+}
+
+void AChaosDataflowSolverActor::RegisterPhysicsComponent(UPrimitiveComponent* Component)
+{
+	if (!Component || !SolverContext)
+	{
+		return;
+	}
+
+	if (FRigidPhysicsSceneDataflowState* State = SolverContext->GetRigidState())
+	{
+		URigidAssetUserData* Data = NewObject<URigidAssetUserData>(Component, NAME_None, RF_Transient);
+
+		if (UInstancedStaticMeshComponent* InstComponent = Cast<UInstancedStaticMeshComponent>(Component))
+		{
+			InstancedRigidComponentInterface::OnCreateSolverBodies(InstComponent, State->Handle, Data->Bodies);
+		}
+		else if (ULandscapeHeightfieldCollisionComponent* LandscapeComponent = Cast<ULandscapeHeightfieldCollisionComponent>(Component))
+		{
+		}
+		else if (USkeletalMeshComponent* SkelComponent = Cast<USkeletalMeshComponent>(Component))
+		{
+			SkeletalRigidComponentInterface::OnCreateSolverBodies(SkelComponent, State->Handle, Data->Bodies);
+		}
+		else
+		{
+			PrimitiveRigidComponentInterface::OnCreateSolverBodies(Component, State->Handle, Data->Bodies);
+		}
+
+		Component->AddAssetUserData(Data);
+
+		RegisteredComponents.Add(Component);
+	}
+}
+
+void AChaosDataflowSolverActor::UnregisterPhysicsComponent(UPrimitiveComponent* Component)
+{
+	if (!Component || !SolverContext)
+	{
+		return;
+	}
+
+	int32 Index = RegisteredComponents.Find(Component);
+	
+	if (Index != INDEX_NONE)
+	{
+		if (FRigidPhysicsSceneDataflowState* State = SolverContext->GetRigidState())
+		{
+			if (UE::Physics::FRigidContextGameRW Context = State->Handle.LockRW())
+			{
+				URigidAssetUserData* RigidData = Component->GetAssetUserData<URigidAssetUserData>();
+				if (RigidData)
+				{
+					for (UE::Physics::FRigidBodyHandle& Handle : RigidData->Bodies)
+					{
+						Context->DestroyBody(Handle.Pin(Context));
+					}
+				}
+
+				Component->RemoveUserDataOfClass(URigidAssetUserData::StaticClass());
+				RegisteredComponents.RemoveAtSwap(Index);
+			}
+		}
+	}
+}
+
+
+
+
diff --git a/Engine/Plugins/Experimental/ChaosDataflowSolver/Source/ChaosDataflowSolver/Private/ChaosSolverBindingComponent.cpp b/Engine/Plugins/Experimental/ChaosDataflowSolver/Source/ChaosDataflowSolver/Private/ChaosSolverBindingComponent.cpp
new file mode 100644
--- /dev/null
+++ b/Engine/Plugins/Experimental/ChaosDataflowSolver/Source/ChaosDataflowSolver/Private/ChaosSolverBindingComponent.cpp
@@ -0,0 +1,101 @@
+// Copyright Epic Games, Inc. All Rights Reserved.
+
+#include "ChaosSolverBindingComponent.h"
+#include "Chaos/DebugDrawQueue.h"
+#include "ChaosLog.h"
+#include "Debug/DebugDrawService.h"
+#include "DrawDebugHelpers.h"
+#if WITH_EDITOR
+#include "UObject/UObjectIterator.h"
+#endif
+#include "Engine/Canvas.h"
+#include "Engine/Engine.h"
+#include "Engine/HitResult.h"
+#include "Engine/World.h"
+#include "GameFramework/Pawn.h"
+#include "GameFramework/PlayerController.h"
+#include "Physics/Experimental/PhysScene_Chaos.h"
+#include "VisualLogger/VisualLogger.h"
+
+#include UE_INLINE_GENERATED_CPP_BY_NAME(ChaosSolverBindingComponent)
+
+UChaosSolverBindingComponent::UChaosSolverBindingComponent()
+{
+}
+
+void UChaosSolverBindingComponent::BeginDestroy()
+{
+	Super::BeginDestroy();
+}
+
+void UChaosSolverBindingComponent::BeginPlay()
+{
+	Super::BeginPlay();
+	if (AActor* Owner = GetOwner())
+	{
+		TSet<UActorComponent*> Components = Owner->GetComponents();
+		
+		for (UActorComponent* Component : Components)
+		{
+			if (TObjectPtr<UPrimitiveComponent> SimulatingComponent = Cast<UPrimitiveComponent>(Component))
+			{				
+				if (SimulationActor)
+				{
+					SimulationActor->RegisterPhysicsComponent(SimulatingComponent);
+
+					if (!bKeepKinematicInOriginal)
+					{
+						Component->DestroyPhysicsState();
+					}
+				}
+			}
+		}
+	}
+
+}
+
+void UChaosSolverBindingComponent::EndPlay(const EEndPlayReason::Type EndPlayReason)
+{
+	Super::EndPlay(EndPlayReason);
+
+	if (AActor* Owner = GetOwner())
+	{
+		TSet<UActorComponent*> Components = Owner->GetComponents();
+		
+		for (UActorComponent* Component : Components)
+		{
+			if (TObjectPtr<UPrimitiveComponent> SimulatingComponent = Cast<UPrimitiveComponent>(Component))
+			{				
+				if (SimulationActor)
+				{
+					SimulationActor->UnregisterPhysicsComponent(SimulatingComponent);
+				}
+			}
+		}
+	}
+}
+
+#if WITH_EDITOR
+void UChaosSolverBindingComponent::OnRegister()
+{
+	Super::OnRegister();
+}
+
+void UChaosSolverBindingComponent::OnUnregister()
+{
+	Super::OnUnregister();
+}
+#endif // WITH_EDITOR
+
+void UChaosSolverBindingComponent::BindWorldDelegates()
+{
+}
+
+void UChaosSolverBindingComponent::HandlePostWorldInitialization(UWorld* World, const UWorld::InitializationValues)
+{
+}
+
+
+
+
+
diff --git a/Engine/Plugins/Experimental/ChaosDataflowSolver/Source/ChaosDataflowSolver/Private/RigidComponents/InstancedRigidComponentInterface.h b/Engine/Plugins/Experimental/ChaosDataflowSolver/Source/ChaosDataflowSolver/Private/RigidComponents/InstancedRigidComponentInterface.h
new file mode 100644
--- /dev/null
+++ b/Engine/Plugins/Experimental/ChaosDataflowSolver/Source/ChaosDataflowSolver/Private/RigidComponents/InstancedRigidComponentInterface.h
@@ -0,0 +1,89 @@
+// Copyright Epic Games, Inc. All Rights Reserved.
+
+#pragma once
+
+#if UE_RIGIDPHYSICS_API_ENABLED
+
+#include "Components/InstancedStaticMeshComponent.h"
+
+#include "RigidPhysics/RigidBody.h"
+#include "RigidPhysics/RigidContext.h"
+#include "RigidPhysics/RigidModifier.h"
+#include "RigidPhysics/RigidScene.h"
+#include "PhysicsEngine/BodyInstance.h"
+#include "PhysicsEngine/BodySetup.h"
+#include "Physics/Experimental/ChaosInterfaceUtils.h"
+
+class InstancedRigidComponentInterface
+{
+	public:
+	static void OnCreateSolverBodies(UInstancedStaticMeshComponent* Component, UE::Physics::FRigidSceneHandle& SceneHandle, TArray<UE::Physics::FRigidBodyHandle>& OutBodies)
+	{
+		using namespace Chaos;
+		AActor* Owner = Component->GetOwner();
+
+		const TArray<FBodyInstance*>& OriginalBodies = Component->InstanceBodies;
+	
+		for (int i = 0; i < OriginalBodies.Num(); i++)
+		{
+			if (const FBodyInstance* OriginalBody = OriginalBodies[i])
+			{
+				if (UBodySetup* BodySetup = OriginalBody->GetBodySetup())
+				{
+					if (UE::Physics::FRigidContextGameRW Context = SceneHandle.LockRW())
+					{
+						UE::Physics::FRigidBodyHandle Body = Context->CreateBody(UE::Physics::FRigidDebugName(OriginalBody->GetBodyDebugName()), UE::Physics::ERigidMovementType::Static);
+
+						if (UE::Physics::TRigidBodyPtr<UE::Physics::FRigidContextGameRW> BodyPtr = Body.Pin(Context))
+						{
+							FTransform Transform = FTransform::Identity;
+
+							Component->GetInstanceTransform(i, Transform, true);
+
+							FBodyCollisionData BodyCollisionData;
+							OriginalBody->BuildBodyFilterData(BodyCollisionData.CollisionFilterData);
+							FBodyInstance::BuildBodyCollisionFlags(BodyCollisionData.CollisionFlags, OriginalBody->GetCollisionEnabled(), BodySetup->GetCollisionTraceFlag() == CTF_UseComplexAsSimple);
+
+							TArray<FPhysicalMaterialMaskParams> OutPhysMaterialMasks;
+
+							FGeometryAddParams AddParams;
+							AddParams.bDoubleSided = OriginalBody->GetBodySetup()->bDoubleSidedGeometry;
+							AddParams.CollisionData = BodyCollisionData; 
+							AddParams.CollisionTraceType = BodySetup->GetCollisionTraceFlag();
+							AddParams.SimpleMaterial = OriginalBody->GetSimplePhysicalMaterial();
+							AddParams.ComplexMaterials = TArrayView<UPhysicalMaterial*>(OriginalBody->GetComplexPhysicalMaterials(OutPhysMaterialMasks));
+							AddParams.ComplexMaterialMasks = TArrayView<FPhysicalMaterialMaskParams>(OutPhysMaterialMasks);
+							AddParams.Scale = Transform.GetScale3D();
+							AddParams.LocalTransform = FTransform::Identity;
+							AddParams.WorldTransform = Transform;
+							AddParams.Geometry = &BodySetup->AggGeom;
+							AddParams.TriMeshGeometries = MakeArrayView(BodySetup->TriMeshGeometries);
+
+							TArray<Chaos::FImplicitObjectPtr> Geoms;
+							Chaos::FShapesArray Shapes;
+							ChaosInterface::CreateGeometry(AddParams, Geoms, Shapes);
+
+							const int32 NumShapes = Shapes.Num();
+		
+							for (int32 ShapeIndex = 0; ShapeIndex < NumShapes; ShapeIndex++)
+							{
+								BodyPtr->AddShape(Shapes[ShapeIndex].Release());
+							}
+
+							BodyPtr->InitTransform(Transform);
+
+							BodyPtr->SetMass(OriginalBody->GetBodyMass());
+							BodyPtr->SetInertia(OriginalBody->GetBodyInertiaTensor() * OriginalBody->InertiaTensorScale);
+
+							BodyPtr->Activate();
+						}
+
+						OutBodies.Add(Body);
+					}
+				}
+			}
+		}
+	}
+};
+
+#endif
\ No newline at end of file
diff --git a/Engine/Plugins/Experimental/ChaosDataflowSolver/Source/ChaosDataflowSolver/Private/RigidComponents/LandscapeRigidComponentInterface.h b/Engine/Plugins/Experimental/ChaosDataflowSolver/Source/ChaosDataflowSolver/Private/RigidComponents/LandscapeRigidComponentInterface.h
new file mode 100644
--- /dev/null
+++ b/Engine/Plugins/Experimental/ChaosDataflowSolver/Source/ChaosDataflowSolver/Private/RigidComponents/LandscapeRigidComponentInterface.h
@@ -0,0 +1,19 @@
+// Copyright Epic Games, Inc. All Rights Reserved.
+
+#pragma once
+
+#if UE_RIGIDPHYSICS_API_ENABLED
+
+#include "LandscapeHeightfieldCollisionComponent.h"
+
+#include "RigidPhysics/RigidBody.h"
+#include "RigidPhysics/RigidContext.h"
+#include "RigidPhysics/RigidModifier.h"
+#include "RigidPhysics/RigidScene.h"
+
+class LandscapeRigidComponentInterface
+{
+
+};
+
+#endif
\ No newline at end of file
diff --git a/Engine/Plugins/Experimental/ChaosDataflowSolver/Source/ChaosDataflowSolver/Private/RigidComponents/PrimitiveRigidComponentInterface.h b/Engine/Plugins/Experimental/ChaosDataflowSolver/Source/ChaosDataflowSolver/Private/RigidComponents/PrimitiveRigidComponentInterface.h
new file mode 100644
--- /dev/null
+++ b/Engine/Plugins/Experimental/ChaosDataflowSolver/Source/ChaosDataflowSolver/Private/RigidComponents/PrimitiveRigidComponentInterface.h
@@ -0,0 +1,125 @@
+// Copyright Epic Games, Inc. All Rights Reserved.
+
+#pragma once
+
+#if UE_RIGIDPHYSICS_API_ENABLED
+
+#include "Components/PrimitiveComponent.h"
+
+#include "RigidPhysics/RigidBody.h"
+#include "RigidPhysics/RigidContext.h"
+#include "RigidPhysics/RigidModifier.h"
+#include "RigidPhysics/RigidScene.h"
+#include "PhysicsEngine/BodyInstance.h"
+#include "PhysicsEngine/BodySetup.h"
+#include "Physics/Experimental/ChaosInterfaceUtils.h"
+
+class PrimitiveRigidComponentInterface
+{
+	public:
+	static void OnCreateSolverBodies(UPrimitiveComponent* Component, UE::Physics::FRigidSceneHandle& SceneHandle, TArray<UE::Physics::FRigidBodyHandle>& OutBodies)
+	{
+		using namespace Chaos;
+		AActor* Owner = Component->GetOwner();
+
+		if (FBodyInstance* OriginalBody = Component->GetBodyInstance())
+		{
+			if (UBodySetup* BodySetup = OriginalBody->GetBodySetup())
+			{
+				if (UE::Physics::FRigidContextGameRW Context = SceneHandle.LockRW())
+				{
+					UE::Physics::ERigidMovementType MovementType;
+
+					if (!OriginalBody->IsNonKinematic())
+					{
+						MovementType = UE::Physics::ERigidMovementType::Kinematic;
+					}
+					else
+					{
+						MovementType = UE::Physics::ERigidMovementType::Dynamic;
+					}
+
+					UE::Physics::FRigidBodyHandle Body = Context->CreateBody(UE::Physics::FRigidDebugName(Component->GetName()), MovementType);
+
+					if (UE::Physics::TRigidBodyPtr<UE::Physics::FRigidContextGameRW> BodyPtr = Body.Pin(Context))
+					{
+						FTransform Transform = OriginalBody->GetUnrealWorldTransform();
+
+						FBodyCollisionData BodyCollisionData;
+						OriginalBody->BuildBodyFilterData(BodyCollisionData.CollisionFilterData);
+						FBodyInstance::BuildBodyCollisionFlags(BodyCollisionData.CollisionFlags, OriginalBody->GetCollisionEnabled(), BodySetup->GetCollisionTraceFlag() == CTF_UseComplexAsSimple);
+
+						TArray<FPhysicalMaterialMaskParams> OutPhysMaterialMasks;
+
+						FGeometryAddParams AddParams;
+						AddParams.bDoubleSided = OriginalBody->GetBodySetup()->bDoubleSidedGeometry;
+						AddParams.CollisionData = BodyCollisionData;
+						AddParams.CollisionTraceType = BodySetup->GetCollisionTraceFlag();
+						AddParams.SimpleMaterial = OriginalBody->GetSimplePhysicalMaterial();
+						AddParams.ComplexMaterials = TArrayView<UPhysicalMaterial*>(OriginalBody->GetComplexPhysicalMaterials(OutPhysMaterialMasks));
+						AddParams.ComplexMaterialMasks = TArrayView<FPhysicalMaterialMaskParams>(OutPhysMaterialMasks);
+						AddParams.Scale = Component->GetComponentScale();
+						AddParams.LocalTransform = FTransform::Identity;
+						AddParams.WorldTransform = Transform;
+						AddParams.Geometry = &BodySetup->AggGeom;
+						AddParams.TriMeshGeometries = MakeArrayView(BodySetup->TriMeshGeometries);
+
+						TArray<Chaos::FImplicitObjectPtr> Geoms;
+						Chaos::FShapesArray Shapes;
+						ChaosInterface::CreateGeometry(AddParams, Geoms, Shapes);
+
+						const int32 NumShapes = Shapes.Num();
+
+						for (int32 ShapeIndex = 0; ShapeIndex < NumShapes; ++ShapeIndex)
+						{
+							BodyPtr->AddShape(Shapes[ShapeIndex].Release());
+						}
+
+						Transform.SetScale3D(FVector(1.0f, 1.0f, 1.0f));
+						BodyPtr->InitTransform(Transform);
+
+						BodyPtr->SetMass(OriginalBody->GetBodyMass());
+						BodyPtr->SetInertia(OriginalBody->GetBodyInertiaTensor() * OriginalBody->InertiaTensorScale);
+
+						BodyPtr->Activate();
+					}
+
+					OutBodies.Add(Body);
+
+					OriginalBody->SetInstanceSimulatePhysics(false);
+				}
+			}
+		}
+	}
+
+	static void OnSolverEndFrame(UPrimitiveComponent* Component, UE::Physics::FRigidSceneHandle& SceneHandle, TArray<UE::Physics::FRigidBodyHandle>& Bodies)
+	{
+		if (Bodies.IsEmpty())
+		{
+			return;
+		}
+
+		//Only 1 body;
+		UE::Physics::FRigidBodyHandle BodyHandle = Bodies[0];
+
+		if (UE::Physics::FRigidContextGameRO Context = SceneHandle.LockRO())
+		{
+			if (UE::Physics::TRigidBodyPtr<UE::Physics::FRigidContextGameRW> BodyPtr = BodyHandle.Pin(Context))
+			{
+				if (!BodyPtr->IsStatic())
+				{
+					if (!BodyPtr->IsKinematic())
+					{
+						FTransform NewTransform = BodyPtr->GetTransform();
+
+						const FVector DesiredDelta = FTransform::SubtractTranslations(NewTransform, Component->GetComponentTransform());
+
+						Component->MoveComponent(DesiredDelta, NewTransform.GetRotation(), false, NULL, MOVECOMP_SkipPhysicsMove);
+					}
+				}
+			}
+		}
+	}
+};
+
+#endif
\ No newline at end of file
diff --git a/Engine/Plugins/Experimental/ChaosDataflowSolver/Source/ChaosDataflowSolver/Private/RigidComponents/SkeletalRigidComponentInterface.h b/Engine/Plugins/Experimental/ChaosDataflowSolver/Source/ChaosDataflowSolver/Private/RigidComponents/SkeletalRigidComponentInterface.h
new file mode 100644
--- /dev/null
+++ b/Engine/Plugins/Experimental/ChaosDataflowSolver/Source/ChaosDataflowSolver/Private/RigidComponents/SkeletalRigidComponentInterface.h
@@ -0,0 +1,26 @@
+// Copyright Epic Games, Inc. All Rights Reserved.
+
+#pragma once
+
+#include "Components/SkeletalMeshComponent.h"
+
+#include "RigidPhysics/RigidBody.h"
+#include "RigidPhysics/RigidContext.h"
+#include "RigidPhysics/RigidModifier.h"
+#include "RigidPhysics/RigidScene.h"
+
+#if UE_RIGIDPHYSICS_API_ENABLED
+
+class SkeletalRigidComponentInterface
+{
+	public:
+	static void OnCreateSolverBodies(USkeletalMeshComponent* Component, UE::Physics::FRigidSceneHandle& SceneHandle, TArray<UE::Physics::FRigidBodyHandle>& OutBodies)
+	{
+	}
+
+	static void OnSolverEndFrame(USkeletalMeshComponent* Component, UE::Physics::FRigidSceneHandle& SceneHandle, TArray<UE::Physics::FRigidBodyHandle>& Bodies)
+	{
+	}
+};
+
+#endif
\ No newline at end of file
diff --git a/Engine/Plugins/Experimental/ChaosDataflowSolver/Source/ChaosDataflowSolver/Public/ChaosDataflowSolver.h b/Engine/Plugins/Experimental/ChaosDataflowSolver/Source/ChaosDataflowSolver/Public/ChaosDataflowSolver.h
new file mode 100644
--- /dev/null
+++ b/Engine/Plugins/Experimental/ChaosDataflowSolver/Source/ChaosDataflowSolver/Public/ChaosDataflowSolver.h
@@ -0,0 +1,14 @@
+// Copyright Epic Games, Inc. All Rights Reserved.
+
+#pragma once
+
+#include "Modules/ModuleManager.h"
+
+class FChaosDataflowSolverModule : public IModuleInterface
+{
+public:
+
+	/** IModuleInterface implementation */
+	virtual void StartupModule() override;
+	virtual void ShutdownModule() override;
+};
diff --git a/Engine/Plugins/Experimental/ChaosDataflowSolver/Source/ChaosDataflowSolver/Public/ChaosDataflowSolverActor.h b/Engine/Plugins/Experimental/ChaosDataflowSolver/Source/ChaosDataflowSolver/Public/ChaosDataflowSolverActor.h
new file mode 100644
--- /dev/null
+++ b/Engine/Plugins/Experimental/ChaosDataflowSolver/Source/ChaosDataflowSolver/Public/ChaosDataflowSolverActor.h
@@ -0,0 +1,60 @@
+// Copyright Epic Games, Inc. All Rights Reserved.
+
+#pragma once
+
+#include "CoreMinimal.h"
+#include "GameFramework/Actor.h"
+#include "Dataflow/Interfaces/DataflowPhysicsSolver.h"
+#include "Dataflow/DataflowSimulationInterface.h"
+#include "UObject/ObjectMacros.h"
+#include "Engine/EngineTypes.h"
+#include "RigidPhysics/RigidFactory.h"
+#include "RigidPhysics/RigidScene.h"
+#include "RigidPhysics/RigidBody.h"
+#include "Dataflow/RigidPhysicsSceneState.h"
+#include "Dataflow/DataflowSimulationManager.h"
+
+#include "Components/InstancedStaticMeshComponent.h"
+
+#include "ChaosDataflowSolverActor.generated.h"
+
+struct FDataflowSimulationAsset;
+class UChaosSolverBindingComponent;
+
+UCLASS(meta=(Experimental), MinimalAPI)
+class AChaosDataflowSolverActor : public AActor, public IDataflowPhysicsSolverInterface
+{
+	GENERATED_UCLASS_BODY()
+
+public:
+	virtual void PostInitializeComponents() override;
+	virtual void BeginPlay() override;
+	virtual void EndPlay(const EEndPlayReason::Type ReasonEnd) override;
+	/** Controls whether the solver is able to simulate particles it controls */
+	UFUNCTION(BlueprintCallable, Category = "ChaosPhysics")
+	virtual void SetSolverActive(bool bActive);
+
+	//~ Begin IDataflowPhysicsSolverInterface interface
+	virtual FString GetSimulationName() const override {return GetName();};
+	virtual FDataflowSimulationAsset& GetSimulationAsset() override {return SimulationAsset;};
+	virtual const FDataflowSimulationAsset& GetSimulationAsset() const override {return SimulationAsset;};
+	virtual FDataflowSimulationProxy* GetSimulationProxy() override { return nullptr; }
+	virtual const FDataflowSimulationProxy* GetSimulationProxy() const  override { return nullptr; }
+	virtual void BuildSimulationProxy() override;
+    virtual void ResetSimulationProxy() override;
+	virtual void WriteToSimulation(const float DeltaTime, const bool bAsyncTask) override;
+	virtual void ReadFromSimulation(const float DeltaTime, const bool bAsyncTask) override;
+	virtual void Tick(float DeltaTime) override;
+
+	void RegisterPhysicsComponent(UPrimitiveComponent* Component);
+	void UnregisterPhysicsComponent(UPrimitiveComponent* Component);
+
+private:
+	/* Solver dataflow asset used to advance in time */
+	UPROPERTY(EditAnywhere, Category = "Physics", meta=(EditConditionHides))
+	FDataflowSimulationAsset SimulationAsset;
+
+	TArray<TWeakObjectPtr<UPrimitiveComponent>> RegisteredComponents;
+
+	TUniquePtr<UE::Dataflow::FDataflowSimulationContext> SolverContext;
+};
\ No newline at end of file
diff --git a/Engine/Plugins/Experimental/ChaosDataflowSolver/Source/ChaosDataflowSolver/Public/ChaosSolverBindingComponent.h b/Engine/Plugins/Experimental/ChaosDataflowSolver/Source/ChaosDataflowSolver/Public/ChaosSolverBindingComponent.h
new file mode 100644
--- /dev/null
+++ b/Engine/Plugins/Experimental/ChaosDataflowSolver/Source/ChaosDataflowSolver/Public/ChaosSolverBindingComponent.h
@@ -0,0 +1,45 @@
+// Copyright Epic Games, Inc. All Rights Reserved.
+
+#pragma once
+
+#include "Components/ActorComponent.h"
+#include "ChaosDataflowSolverActor.h"
+#include "Chaos/ChaosDebugDrawDeclares.h"
+#include "Chaos/Declares.h"
+#include "Chaos/DebugDrawQueue.h"
+#include "Engine/World.h"
+#include "ChaosSolverBindingComponent.generated.h"
+
+UCLASS(BlueprintType, ClassGroup = Chaos, meta = (Experimental, BlueprintSpawnableComponent), MinimalAPI)
+class UChaosSolverBindingComponent: public UActorComponent
+{
+	GENERATED_BODY()
+
+public:
+
+	UChaosSolverBindingComponent();
+
+	//~ Begin UActorComponent interface
+	virtual void BeginPlay() override;
+	virtual void EndPlay(const EEndPlayReason::Type EndPlayReason) override;
+	virtual void BeginDestroy() override;
+	//~ End UActorComponent interface
+
+	static void BindWorldDelegates();
+
+protected:
+#if WITH_EDITOR
+	virtual void OnRegister() override;
+	virtual void OnUnregister() override;
+#endif // WITH_EDITOR
+
+private:
+	static void HandlePostWorldInitialization(UWorld* World, const UWorld::InitializationValues IVS);
+
+	UPROPERTY(EditAnywhere, Category = "Physics")
+	TSoftObjectPtr<AChaosDataflowSolverActor> SimulationActor;
+
+	UPROPERTY(EditAnywhere, Category = "Physics")
+	bool bKeepKinematicInOriginal;
+};
+
diff --git a/Engine/Plugins/Experimental/ChaosDataflowSolver/Source/ChaosDataflowSolver/Public/RigidAssetUserData.h b/Engine/Plugins/Experimental/ChaosDataflowSolver/Source/ChaosDataflowSolver/Public/RigidAssetUserData.h
new file mode 100644
--- /dev/null
+++ b/Engine/Plugins/Experimental/ChaosDataflowSolver/Source/ChaosDataflowSolver/Public/RigidAssetUserData.h
@@ -0,0 +1,19 @@
+// Copyright Epic Games, Inc. All Rights Reserved.
+
+#pragma once
+
+#include "Engine/AssetUserData.h"
+#include "CoreMinimal.h"
+#include "RigidPhysics/RigidBody.h"
+
+#include "RigidAssetUserData.generated.h"
+
+
+UCLASS(meta=(Experimental), MinimalAPI, NotBlueprintable, HideCategories = (Object))
+class URigidAssetUserData : public UAssetUserData
+{
+	GENERATED_BODY()
+
+public:
+	TArray<UE::Physics::FRigidBodyHandle> Bodies;
+};
\ No newline at end of file
diff --git a/Engine/Plugins/Experimental/ChaosRigidPhysicsAsync/Source/ChaosRigidPhysicsAsync/ChaosRigidPhysicsAsync.Build.cs b/Engine/Plugins/Experimental/ChaosRigidPhysicsAsync/Source/ChaosRigidPhysicsAsync/ChaosRigidPhysicsAsync.Build.cs
--- a/Engine/Plugins/Experimental/ChaosRigidPhysicsAsync/Source/ChaosRigidPhysicsAsync/ChaosRigidPhysicsAsync.Build.cs
+++ b/Engine/Plugins/Experimental/ChaosRigidPhysicsAsync/Source/ChaosRigidPhysicsAsync/ChaosRigidPhysicsAsync.Build.cs
@@ -14,6 +14,9 @@ namespace UnrealBuildTool.Rules
                 "Core",
                 "CoreUObject",
                 "RigidPhysics",
+				"DataflowCore",
+				"DataflowEngine",
+				"DataflowSimulation",
 				}
 			);
 
diff --git a/Engine/Plugins/Experimental/ChaosRigidPhysicsAsync/Source/ChaosRigidPhysicsAsync/ChaosRigidPhysicsAsyncModule.cpp b/Engine/Plugins/Experimental/ChaosRigidPhysicsAsync/Source/ChaosRigidPhysicsAsync/ChaosRigidPhysicsAsyncModule.cpp
--- a/Engine/Plugins/Experimental/ChaosRigidPhysicsAsync/Source/ChaosRigidPhysicsAsync/ChaosRigidPhysicsAsyncModule.cpp
+++ b/Engine/Plugins/Experimental/ChaosRigidPhysicsAsync/Source/ChaosRigidPhysicsAsync/ChaosRigidPhysicsAsyncModule.cpp
@@ -2,4 +2,17 @@
 
 #include "Modules/ModuleManager.h"
 
-IMPLEMENT_MODULE(FDefaultModuleImpl, ChaosRigidPhysicsAsync)
\ No newline at end of file
+#include "ChaosRigidPhysicsAsync/Dataflow/RigidPhysicsAsyncDataflowNodes.h"
+
+class FChaosRigidPhysicsAsyncModule : public FDefaultModuleImpl
+{
+public:
+	virtual void StartupModule() override
+	{
+		RegisterRigidPhysicsAsyncDataflowSimulationNodes();
+	}
+	virtual void ShutdownModule() override
+	{
+	}
+};
+IMPLEMENT_MODULE(FChaosRigidPhysicsAsyncModule, ChaosRigidPhysicsAsync)
\ No newline at end of file
diff --git a/Engine/Plugins/Experimental/ChaosRigidPhysicsAsync/Source/ChaosRigidPhysicsAsync/Private/ChaosRigidPhysicsAsync/Dataflow/RigidPhysicsAsyncDataflowNodes.cpp b/Engine/Plugins/Experimental/ChaosRigidPhysicsAsync/Source/ChaosRigidPhysicsAsync/Private/ChaosRigidPhysicsAsync/Dataflow/RigidPhysicsAsyncDataflowNodes.cpp
new file mode 100644
--- /dev/null
+++ b/Engine/Plugins/Experimental/ChaosRigidPhysicsAsync/Source/ChaosRigidPhysicsAsync/Private/ChaosRigidPhysicsAsync/Dataflow/RigidPhysicsAsyncDataflowNodes.cpp
@@ -0,0 +1,45 @@
+// Copyright Epic Games, Inc. All Rights Reserved.
+
+#include "ChaosRigidPhysicsAsync/Dataflow/RigidPhysicsAsyncDataflowNodes.h"
+#include "Dataflow/Interfaces/DataflowPhysicsSolver.h"
+#include "Dataflow/DataflowSimulationNodes.h"
+#include "Dataflow/DataflowSimulationManager.h"
+#include "Dataflow/DataflowNodeFactory.h"
+#include "Dataflow/DataflowNodeColorsRegistry.h"
+#include "ChaosRigidPhysicsAsync/RigidFactoryAsync.h"
+#include "ChaosRigidPhysicsAsync/RigidBodyAsync.h"
+#include "ChaosRigidPhysicsAsync/RigidSceneSettingsAsync.h"
+#include "RigidPhysics/RigidBody.h"
+#include "RigidPhysics/RigidContext.h"
+#include "RigidPhysics/RigidModifier.h"
+#include "RigidPhysics/RigidScene.h"
+
+#include UE_INLINE_GENERATED_CPP_BY_NAME(RigidPhysicsAsyncDataflowNodes)
+void RegisterRigidPhysicsAsyncDataflowSimulationNodes()
+{
+#if UE_RIGIDPHYSICS_API_ENABLED
+	DATAFLOW_NODE_REGISTER_CREATION_FACTORY(FRigidSceneSetupDataflowNode);
+#endif
+}
+
+void FRigidSceneSetupDataflowNode::Evaluate(UE::Dataflow::FContext& Context, const FDataflowOutput* Output) const
+{
+#if UE_RIGIDPHYSICS_API_ENABLED
+	using namespace Chaos;
+	using namespace Chaos::Rigids::Async;
+	using namespace UE::Physics;
+
+	if(Context.IsA(UE::Dataflow::FDataflowSimulationContext::StaticType()))
+	{
+		UE::Dataflow::FDataflowSimulationContext& SimulationContext = StaticCast<UE::Dataflow::FDataflowSimulationContext&>(Context);
+		if (Output->IsA(&State))
+		{
+			FRigidPhysicsSceneDataflowState InState;
+			InState.Factory = MakeShareable(new FRigidFactoryAsync);
+			FRigidSceneSettingsAsync SceneSettings;
+			InState.Handle = InState.Factory->CreateScene(FRigidDebugName(SimulationContext.GetRigidSceneName()), &SceneSettings);
+			SetValue(SimulationContext, InState, &State);
+		}
+	}
+#endif
+}
\ No newline at end of file
diff --git a/Engine/Plugins/Experimental/ChaosRigidPhysicsAsync/Source/ChaosRigidPhysicsAsync/Public/ChaosRigidPhysicsAsync/Dataflow/RigidPhysicsAsyncDataflowNodes.h b/Engine/Plugins/Experimental/ChaosRigidPhysicsAsync/Source/ChaosRigidPhysicsAsync/Public/ChaosRigidPhysicsAsync/Dataflow/RigidPhysicsAsyncDataflowNodes.h
new file mode 100644
--- /dev/null
+++ b/Engine/Plugins/Experimental/ChaosRigidPhysicsAsync/Source/ChaosRigidPhysicsAsync/Public/ChaosRigidPhysicsAsync/Dataflow/RigidPhysicsAsyncDataflowNodes.h
@@ -0,0 +1,31 @@
+// Copyright Epic Games, Inc. All Rights Reserved.
+#pragma once
+#include "Dataflow/DataflowObject.h"
+#include "Dataflow/DataflowSimulationProxy.h"
+#include "Dataflow/DataflowSimulationContext.h"
+#include "Dataflow/DataflowSimulationNodes.h"
+#include "Dataflow/RigidPhysicsSceneState.h"
+#include "RigidPhysics/RigidScene.h"
+#include "RigidPhysicsAsyncDataflowNodes.generated.h"
+using namespace UE::Physics;
+/** Create a rigid scene */
+USTRUCT(meta = (DataflowSimulation))
+struct FRigidSceneSetupDataflowNode : public FDataflowNode
+{
+	GENERATED_USTRUCT_BODY()
+	DATAFLOW_NODE_DEFINE_INTERNAL(FRigidSceneSetupDataflowNode, "CreateRigidScene", "Physics|Scene", UDataflow::SimulationTag)
+public:
+	
+	FRigidSceneSetupDataflowNode(const UE::Dataflow::FNodeParameters& InParam, FGuid InGuid = FGuid::NewGuid())
+	: FDataflowNode(InParam, InGuid)
+	{
+		RegisterOutputConnection(&State);
+	}
+
+	UPROPERTY(meta = (DataflowOutput))
+	FRigidPhysicsSceneDataflowState State;
+	virtual void Evaluate(UE::Dataflow::FContext& Context, const FDataflowOutput* Output) const override;
+	//~ End UDataflowNode interface
+};
+
+void RegisterRigidPhysicsAsyncDataflowSimulationNodes();
\ No newline at end of file
diff --git a/Engine/Source/Runtime/Experimental/ChaosSolverEngine/ChaosSolverEngine.Build.cs b/Engine/Source/Runtime/Experimental/ChaosSolverEngine/ChaosSolverEngine.Build.cs
--- a/Engine/Source/Runtime/Experimental/ChaosSolverEngine/ChaosSolverEngine.Build.cs
+++ b/Engine/Source/Runtime/Experimental/ChaosSolverEngine/ChaosSolverEngine.Build.cs
@@ -26,7 +26,8 @@ namespace UnrealBuildTool.Rules
 			PrivateDependencyModuleNames.AddRange(
 				new string[]
 				{
-					"Messaging"
+					"Messaging",
+					"RigidPhysics",
 				}
 				);
 
diff --git a/Engine/Source/Runtime/Experimental/Dataflow/Simulation/DataflowSimulation.Build.cs b/Engine/Source/Runtime/Experimental/Dataflow/Simulation/DataflowSimulation.Build.cs
--- a/Engine/Source/Runtime/Experimental/Dataflow/Simulation/DataflowSimulation.Build.cs
+++ b/Engine/Source/Runtime/Experimental/Dataflow/Simulation/DataflowSimulation.Build.cs
@@ -14,6 +14,7 @@ namespace UnrealBuildTool.Rules
 					"CoreUObject",
 					"DataflowCore",
 					"DataflowEngine",
+					"RigidPhysics",
 				}
 			);
 			PrivateDependencyModuleNames.AddRange(
diff --git a/Engine/Source/Runtime/Experimental/Dataflow/Simulation/Private/Dataflow/DataflowSimulationContext.cpp b/Engine/Source/Runtime/Experimental/Dataflow/Simulation/Private/Dataflow/DataflowSimulationContext.cpp
--- a/Engine/Source/Runtime/Experimental/Dataflow/Simulation/Private/Dataflow/DataflowSimulationContext.cpp
+++ b/Engine/Source/Runtime/Experimental/Dataflow/Simulation/Private/Dataflow/DataflowSimulationContext.cpp
@@ -1,7 +1,8 @@
-﻿// Copyright Epic Games, Inc. All Rights Reserved.
+// Copyright Epic Games, Inc. All Rights Reserved.
 
 #include "Dataflow/DataflowSimulationContext.h"
 #include "Dataflow/DataflowSimulationProxy.h"
+#include "Dataflow/RigidPhysicsSceneState.h"
 
 namespace UE::Dataflow
 {
@@ -91,6 +92,33 @@ int32 TSimulationContext<Base>::NumSimulationProxies(const FString& ProxyType) c
 	return TypedProxies ? TypedProxies->Num() : 0;
 }
 
+#if UE_RIGIDPHYSICS_API_ENABLED 
+
+template<typename Base>
+FRigidPhysicsSceneDataflowState* TSimulationContext<Base>::GetRigidState()
+{
+	if (!RigidState)
+	{
+		RigidState = MakePimpl<FRigidPhysicsSceneDataflowState>();
+	}
+
+	return RigidState.Get();
+}
+
+template<typename Base>
+void TSimulationContext<Base>::SetRigidSceneName(const FString& InSceneName)
+{
+	SceneName = InSceneName;
+}
+
+template<typename Base>
+const FString& TSimulationContext<Base>::GetRigidSceneName()
+{
+	return SceneName;
+}
+
+#endif
+
 template class TSimulationContext<UE::Dataflow::FContextSingle>;
 template class TSimulationContext<UE::Dataflow::FContextThreaded>;
 	
diff --git a/Engine/Source/Runtime/Experimental/Dataflow/Simulation/Private/Dataflow/DataflowSimulationNodes.cpp b/Engine/Source/Runtime/Experimental/Dataflow/Simulation/Private/Dataflow/DataflowSimulationNodes.cpp
--- a/Engine/Source/Runtime/Experimental/Dataflow/Simulation/Private/Dataflow/DataflowSimulationNodes.cpp
+++ b/Engine/Source/Runtime/Experimental/Dataflow/Simulation/Private/Dataflow/DataflowSimulationNodes.cpp
@@ -6,6 +6,9 @@
 #include "Dataflow/DataflowNodeColorsRegistry.h"
 #include "Dataflow/Interfaces/DataflowPhysicsSolver.h"
 
+#include "RigidPhysics/RigidContext.h"
+#include "RigidPhysics/RigidScene.h"
+
 #include UE_INLINE_GENERATED_CPP_BY_NAME(DataflowSimulationNodes)
 
 namespace UE::Dataflow
@@ -22,6 +25,14 @@ namespace UE::Dataflow
 		// Proxies nodes
 		DATAFLOW_NODE_REGISTER_CREATION_FACTORY(FFilterSimulationProxiesDataflowNode);
 		DATAFLOW_NODE_REGISTER_CREATION_FACTORY(FSimulationProxiesTerminalDataflowNode);
+
+#if UE_RIGIDPHYSICS_API_ENABLED
+		// Scene Setup
+		DATAFLOW_NODE_REGISTER_CREATION_FACTORY(FRigidSceneSetupTerminalDataflowNode);
+
+		DATAFLOW_NODE_REGISTER_CREATION_FACTORY(FAdvanceRigidPhysicsSolverDataflowNode);
+		DATAFLOW_NODE_REGISTER_CREATION_FACTORY(FSceneSimulationTerminalDataflowNode);
+#endif
 		
 		static constexpr FLinearColor CDefaultNodeBodyTintColor = FLinearColor(0.f, 0.f, 0.0f, 0.5f);
 		
@@ -97,3 +108,37 @@ void FSimulationProxiesTerminalDataflowNode::EvaluateSimulation(UE::Dataflow::FD
 {
 	const TArray<FDataflowSimulationProperty> SolverProperties = GetValue(SimulationContext, &SimulationProxies);
 }
+
+void FAdvanceRigidPhysicsSolverDataflowNode::EvaluateSimulation(UE::Dataflow::FDataflowSimulationContext& SimulationContext, const FDataflowOutput* Output) const
+{
+#if UE_RIGIDPHYSICS_API_ENABLED
+	const FRigidPhysicsSceneDataflowState InRigidState = GetValue(SimulationContext, &State);
+
+	const float SimulationDeltaTime = GetValue(SimulationContext, &SimulationTime).DeltaTime;
+	
+	if(FRigidPhysicsSceneDataflowState* RigidState = SimulationContext.GetRigidState())
+	{
+		if (UE::Physics::FRigidContextGameRW Context = RigidState->Handle.LockRW())
+		{
+			Context->Tick(SimulationDeltaTime);
+		}
+	}
+
+	SetValue(SimulationContext, InRigidState, &State);
+#endif
+}
+
+void FSceneSimulationTerminalDataflowNode::EvaluateSimulation(UE::Dataflow::FDataflowSimulationContext& SimulationContext, const FDataflowOutput* Output) const
+{
+#if UE_RIGIDPHYSICS_API_ENABLED
+	FRigidPhysicsSceneDataflowState InState = GetValue(SimulationContext, &State);
+#endif
+}
+
+void FRigidSceneSetupTerminalDataflowNode::EvaluateSimulation(UE::Dataflow::FDataflowSimulationContext& SimulationContext, const FDataflowOutput* Output) const
+{
+#if UE_RIGIDPHYSICS_API_ENABLED
+	FRigidPhysicsSceneDataflowState InState = GetValue(SimulationContext, &State);
+	*SimulationContext.GetRigidState() = InState;
+#endif
+}
\ No newline at end of file
diff --git a/Engine/Source/Runtime/Experimental/Dataflow/Simulation/Private/Dataflow/DataflowSimulationPlugin.cpp b/Engine/Source/Runtime/Experimental/Dataflow/Simulation/Private/Dataflow/DataflowSimulationPlugin.cpp
--- a/Engine/Source/Runtime/Experimental/Dataflow/Simulation/Private/Dataflow/DataflowSimulationPlugin.cpp
+++ b/Engine/Source/Runtime/Experimental/Dataflow/Simulation/Private/Dataflow/DataflowSimulationPlugin.cpp
@@ -1,4 +1,4 @@
-﻿// Copyright Epic Games, Inc. All Rights Reserved.
+// Copyright Epic Games, Inc. All Rights Reserved.
 
 #include "Dataflow/DataflowSimulationPlugin.h"
 #include "Dataflow/DataflowSimulationNodes.h"
@@ -16,6 +16,11 @@ void IDataflowSimulationPlugin::StartupModule()
 	UE::Dataflow::RegisterNodeFilter(FDataflowInvalidNode::StaticType());
 	UE::Dataflow::RegisterNodeFilter(FDataflowExecutionNode::StaticType());
 
+#if UE_RIGIDPHYSICS_API_ENABLED
+	UE::Dataflow::RegisterNodeFilter(FRigidSceneSetupTerminalDataflowNode::StaticType());
+	UE::Dataflow::RegisterNodeFilter(FSceneSimulationTerminalDataflowNode::StaticType());
+#endif
+
 	UDataflowSimulationManager::OnStartup();
 }
 
diff --git a/Engine/Source/Runtime/Experimental/Dataflow/Simulation/Public/Dataflow/DataflowSimulationContext.h b/Engine/Source/Runtime/Experimental/Dataflow/Simulation/Public/Dataflow/DataflowSimulationContext.h
--- a/Engine/Source/Runtime/Experimental/Dataflow/Simulation/Public/Dataflow/DataflowSimulationContext.h
+++ b/Engine/Source/Runtime/Experimental/Dataflow/Simulation/Public/Dataflow/DataflowSimulationContext.h
@@ -7,6 +7,10 @@
 #include "Dataflow/DataflowObjectInterface.h"
 #include "Dataflow/DataflowSimulationProxy.h"
 
+#include "RigidPhysics/RigidFwd.h"
+
+struct FRigidPhysicsSceneDataflowState;
+
 namespace UE::Dataflow
 {
 	/** Simulation context that will be used by all the simulation nodes*/
@@ -56,6 +60,17 @@ namespace UE::Dataflow
 		/** Build group bits based on the string array */
 		DATAFLOWSIMULATION_API void BuildGroupBits(const TArray<FString>& SimulationGroups, TBitArray<>& GroupBits) const;
 
+#if UE_RIGIDPHYSICS_API_ENABLED 
+		/** Return state associated with the rigids dataflow simulation */
+		UE_INTERNAL DATAFLOWSIMULATION_API FRigidPhysicsSceneDataflowState* GetRigidState();
+
+		/** Set the name of scene created through the rigids dataflow */
+		UE_INTERNAL DATAFLOWSIMULATION_API void SetRigidSceneName(const FString& InSceneName);
+
+		/** Get the name of scene created through the rigids dataflow */
+		UE_INTERNAL DATAFLOWSIMULATION_API const FString& GetRigidSceneName();
+#endif
+
 		/** Push another level of iteration indices */
 		void PushIterationIndex()
 		{
@@ -96,6 +111,11 @@ namespace UE::Dataflow
 
 		/** List of all the nested iteration indices */
 		TArray<int32> IterationIndices;
+
+#if UE_RIGIDPHYSICS_API_ENABLED
+		UE_INTERNAL TPimplPtr<FRigidPhysicsSceneDataflowState> RigidState;
+		UE_INTERNAL FString SceneName = "";
+#endif
 	};
 
 	template<typename Base>
diff --git a/Engine/Source/Runtime/Experimental/Dataflow/Simulation/Public/Dataflow/DataflowSimulationNodes.h b/Engine/Source/Runtime/Experimental/Dataflow/Simulation/Public/Dataflow/DataflowSimulationNodes.h
--- a/Engine/Source/Runtime/Experimental/Dataflow/Simulation/Public/Dataflow/DataflowSimulationNodes.h
+++ b/Engine/Source/Runtime/Experimental/Dataflow/Simulation/Public/Dataflow/DataflowSimulationNodes.h
@@ -5,6 +5,7 @@
 #include "Dataflow/DataflowObject.h"
 #include "Dataflow/DataflowSimulationProxy.h"
 #include "Dataflow/DataflowSimulationContext.h"
+#include "Dataflow/RigidPhysicsSceneState.h"
 
 #include "DataflowSimulationNodes.generated.h"
 
@@ -293,6 +294,74 @@ public:
 	virtual void EvaluateSimulation(UE::Dataflow::FDataflowSimulationContext& SimulationContext, const FDataflowOutput* Output) const override;
 };
 
+/** Advance the rigid solver in time */
+USTRUCT(meta = (DataflowSimulation))
+struct FAdvanceRigidPhysicsSolverDataflowNode : public FDataflowSimulationNode
+{
+	GENERATED_BODY()
+	DATAFLOW_NODE_DEFINE_INTERNAL(FAdvanceRigidPhysicsSolverDataflowNode, "AdvanceRigidSolver", "Physics|Solver", UDataflow::SimulationTag)
+
+public:
+	
+	FAdvanceRigidPhysicsSolverDataflowNode(const UE::Dataflow::FNodeParameters& InParam, FGuid InGuid = FGuid::NewGuid())
+	: FDataflowSimulationNode(InParam, InGuid)
+	{
+		RegisterInputConnection(&SimulationTime);
+		RegisterInputConnection(&State);
+		RegisterOutputConnection(&State,&State);
+	}
+
+	/** Delta time to use to advance the solver*/
+	UPROPERTY(Transient, SkipSerialization, Meta = (DataflowInput))
+	FDataflowSimulationTime SimulationTime = FDataflowSimulationTime(0.0f,0.0f);
+
+	UPROPERTY(Transient, SkipSerialization, Meta = (DataflowInput, DataflowOutput, DataflowPassthrough = "State"))
+	FRigidPhysicsSceneDataflowState State;
+
+	/** Evaluate simulation node given a simulation context */
+	virtual void EvaluateSimulation(UE::Dataflow::FDataflowSimulationContext& SimulationContext, const FDataflowOutput* Output) const override;
+};
+
+/** Main terminal node for simulation */
+USTRUCT(meta = (DataflowSimulation))
+struct FSceneSimulationTerminalDataflowNode : public FDataflowSimulationNode
+{
+	GENERATED_BODY()
+	DATAFLOW_NODE_DEFINE_INTERNAL(FSceneSimulationTerminalDataflowNode, "SceneSimulationTerminal", "Terminal|Scene", UDataflow::SimulationTag)
+public:
+	
+	FSceneSimulationTerminalDataflowNode(const UE::Dataflow::FNodeParameters& InParam, FGuid InGuid = FGuid::NewGuid())
+	: FDataflowSimulationNode(InParam, InGuid)
+	{
+		RegisterInputConnection(&State);
+	}
+	UPROPERTY(Transient, SkipSerialization, Meta = (DataflowInput))
+	FRigidPhysicsSceneDataflowState State;
+	
+	/** Evaluate simulation node given a simulation context */
+	virtual void EvaluateSimulation(UE::Dataflow::FDataflowSimulationContext& SimulationContext, const FDataflowOutput* Output) const override;
+};
+
+/** Main terminal node for scene setup */
+USTRUCT(meta = (DataflowSimulation))
+struct FRigidSceneSetupTerminalDataflowNode : public FDataflowSimulationNode
+{
+	GENERATED_BODY()
+	DATAFLOW_NODE_DEFINE_INTERNAL(FRigidSceneSetupTerminalDataflowNode, "RigidSceneSetupTerminal", "Terminal|Scene", UDataflow::SimulationTag)
+public:
+	
+	FRigidSceneSetupTerminalDataflowNode(const UE::Dataflow::FNodeParameters& InParam, FGuid InGuid = FGuid::NewGuid())
+	: FDataflowSimulationNode(InParam, InGuid)
+	{
+		RegisterInputConnection(&State);
+	}
+	UPROPERTY(Transient, SkipSerialization, Meta = (DataflowInput))
+	FRigidPhysicsSceneDataflowState State;
+	
+	/** Evaluate simulation node given a simulation context */
+	virtual void EvaluateSimulation(UE::Dataflow::FDataflowSimulationContext& SimulationContext, const FDataflowOutput* Output) const override;
+};
+
 namespace UE::Dataflow
 {
 	void RegisterDataflowSimulationNodes();
diff --git a/Engine/Source/Runtime/Experimental/Dataflow/Simulation/Public/Dataflow/RigidPhysicsSceneState.h b/Engine/Source/Runtime/Experimental/Dataflow/Simulation/Public/Dataflow/RigidPhysicsSceneState.h
new file mode 100644
--- /dev/null
+++ b/Engine/Source/Runtime/Experimental/Dataflow/Simulation/Public/Dataflow/RigidPhysicsSceneState.h
@@ -0,0 +1,23 @@
+// Copyright Epic Games, Inc. All Rights Reserved.
+#pragma once
+
+#include "RigidPhysics/RigidFwd.h"
+
+#if UE_RIGIDPHYSICS_API_ENABLED
+#include "RigidPhysics/RigidFactory.h"
+#include "RigidPhysics/RigidScene.h"
+#include "RigidPhysics/RigidBody.h"
+#endif
+
+#include "GameFramework/Actor.h"
+#include "RigidPhysicsSceneState.generated.h"
+USTRUCT()
+struct FRigidPhysicsSceneDataflowState
+{
+	GENERATED_BODY()
+public:
+#if UE_RIGIDPHYSICS_API_ENABLED
+	TSharedPtr<UE::Physics::IRigidFactory> Factory;
+	UE::Physics::FRigidSceneHandle Handle;
+#endif
+};
diff --git a/Engine/Source/Runtime/Experimental/RigidPhysics/Public/RigidPhysics/RigidFwd.h b/Engine/Source/Runtime/Experimental/RigidPhysics/Public/RigidPhysics/RigidFwd.h
--- a/Engine/Source/Runtime/Experimental/RigidPhysics/Public/RigidPhysics/RigidFwd.h
+++ b/Engine/Source/Runtime/Experimental/RigidPhysics/Public/RigidPhysics/RigidFwd.h
@@ -98,7 +98,7 @@ namespace UE::Physics
 
 	using FConstRigidShapeInstancePtr = TRigidShapeInstancePtr<true>;
 	using FRigidShapeInstancePtr = TRigidShapeInstancePtr<false>;
-	using FRigidDebugName = Chaos::FSharedDebugName;
+	using FRigidDebugName = ::Chaos::FSharedDebugName;
 	using FRigidSceneWeakPtr = TWeakPtr<IRigidScene, ESPMode::ThreadSafe>;
 }
 
