Include Rigid Headers

#Synced-CL 51265957

diff --git a/Engine/Plugins/Experimental/ChaosDataflowSolver/Source/ChaosDataflowSolver/Private/RigidComponents/InstancedRigidComponentInterface.h b/Engine/Plugins/Experimental/ChaosDataflowSolver/Source/ChaosDataflowSolver/Private/RigidComponents/InstancedRigidComponentInterface.h
--- a/Engine/Plugins/Experimental/ChaosDataflowSolver/Source/ChaosDataflowSolver/Private/RigidComponents/InstancedRigidComponentInterface.h
+++ b/Engine/Plugins/Experimental/ChaosDataflowSolver/Source/ChaosDataflowSolver/Private/RigidComponents/InstancedRigidComponentInterface.h
@@ -2,6 +2,8 @@
 
 #pragma once
 
+#include "RigidPhysics/RigidFwd.h"
+
 #if UE_RIGIDPHYSICS_API_ENABLED
 
 #include "Components/InstancedStaticMeshComponent.h"
diff --git a/Engine/Plugins/Experimental/ChaosDataflowSolver/Source/ChaosDataflowSolver/Private/RigidComponents/LandscapeRigidComponentInterface.h b/Engine/Plugins/Experimental/ChaosDataflowSolver/Source/ChaosDataflowSolver/Private/RigidComponents/LandscapeRigidComponentInterface.h
--- a/Engine/Plugins/Experimental/ChaosDataflowSolver/Source/ChaosDataflowSolver/Private/RigidComponents/LandscapeRigidComponentInterface.h
+++ b/Engine/Plugins/Experimental/ChaosDataflowSolver/Source/ChaosDataflowSolver/Private/RigidComponents/LandscapeRigidComponentInterface.h
@@ -2,6 +2,8 @@
 
 #pragma once
 
+#include "RigidPhysics/RigidFwd.h"
+
 #if UE_RIGIDPHYSICS_API_ENABLED
 
 #include "LandscapeHeightfieldCollisionComponent.h"
diff --git a/Engine/Plugins/Experimental/ChaosDataflowSolver/Source/ChaosDataflowSolver/Private/RigidComponents/PrimitiveRigidComponentInterface.h b/Engine/Plugins/Experimental/ChaosDataflowSolver/Source/ChaosDataflowSolver/Private/RigidComponents/PrimitiveRigidComponentInterface.h
--- a/Engine/Plugins/Experimental/ChaosDataflowSolver/Source/ChaosDataflowSolver/Private/RigidComponents/PrimitiveRigidComponentInterface.h
+++ b/Engine/Plugins/Experimental/ChaosDataflowSolver/Source/ChaosDataflowSolver/Private/RigidComponents/PrimitiveRigidComponentInterface.h
@@ -2,6 +2,8 @@
 
 #pragma once
 
+#include "RigidPhysics/RigidFwd.h"
+
 #if UE_RIGIDPHYSICS_API_ENABLED
 
 #include "Components/PrimitiveComponent.h"
