Fix a deprecation warning in ChaosDataFlowSolver.

#Synced-CL 52508799

diff --git a/Engine/Plugins/Experimental/ChaosDataflowSolver/Source/ChaosDataflowSolver/Private/RigidComponents/InstancedRigidComponentInterface.h b/Engine/Plugins/Experimental/ChaosDataflowSolver/Source/ChaosDataflowSolver/Private/RigidComponents/InstancedRigidComponentInterface.h
--- a/Engine/Plugins/Experimental/ChaosDataflowSolver/Source/ChaosDataflowSolver/Private/RigidComponents/InstancedRigidComponentInterface.h
+++ b/Engine/Plugins/Experimental/ChaosDataflowSolver/Source/ChaosDataflowSolver/Private/RigidComponents/InstancedRigidComponentInterface.h
@@ -24,7 +24,7 @@ class InstancedRigidComponentInterface
 		using namespace Chaos;
 		AActor* Owner = Component->GetOwner();
 
-		const TArray<FBodyInstance*>& OriginalBodies = Component->InstanceBodies;
+		const TArray<FBodyInstance*>& OriginalBodies = Component->GetInstanceBodies();
 	
 		for (int i = 0; i < OriginalBodies.Num(); i++)
 		{
